Merging using Dev-VR_to_Release-4.21

diff --git a/Engine/Plugins/Developer/CLionSourceCodeAccess/Source/CLionSourceCodeAccess/Private/CLionSourceCodeAccessor.cpp b/Engine/Plugins/Developer/CLionSourceCodeAccess/Source/CLionSourceCodeAccess/Private/CLionSourceCodeAccessor.cpp
--- a/Engine/Plugins/Developer/CLionSourceCodeAccess/Source/CLionSourceCodeAccess/Private/CLionSourceCodeAccessor.cpp
+++ b/Engine/Plugins/Developer/CLionSourceCodeAccess/Source/CLionSourceCodeAccess/Private/CLionSourceCodeAccessor.cpp
@@ -122,6 +122,14 @@ bool FCLionSourceCodeAccessor::OpenSolutionAtPath(const FString& InSolutionPath)
 	{
 		CorrectSolutionPath = CorrectSolutionPath.Left(CorrectSolutionPath.Len() - 3);
 	}
+	// UE4 passes the project folder and name, so strip the name off
+	int32 LastPathIndex = CorrectSolutionPath.Find(TEXT("/"), ESearchCase::IgnoreCase, ESearchDir::FromEnd);
+	if (LastPathIndex > -1)
+	{
+		CorrectSolutionPath = CorrectSolutionPath.Left(LastPathIndex + 1);
+	}
+	// Make sure the path is wrapped in "" properly
+	CorrectSolutionPath = FString::Printf(TEXT("\"%sCMakeLists.txt\""), *CorrectSolutionPath);
 
 	return FPlatformProcess::CreateProc(*ExecutablePath, *CorrectSolutionPath, true, true, false, nullptr, 0, nullptr, nullptr).IsValid();
 }
