CLion: fix registry key with the path to the binaries.

- Reported by a licensee (https://udn.unrealengine.com/questions/546976/fix-for-clion-plugin.html)

diff --git a/Engine/Plugins/Developer/CLionSourceCodeAccess/Source/CLionSourceCodeAccess/Private/CLionSourceCodeAccessor.cpp b/Engine/Plugins/Developer/CLionSourceCodeAccess/Source/CLionSourceCodeAccess/Private/CLionSourceCodeAccessor.cpp
--- a/Engine/Plugins/Developer/CLionSourceCodeAccess/Source/CLionSourceCodeAccess/Private/CLionSourceCodeAccessor.cpp
+++ b/Engine/Plugins/Developer/CLionSourceCodeAccess/Source/CLionSourceCodeAccess/Private/CLionSourceCodeAccessor.cpp
@@ -185,7 +185,7 @@ FString FCLionSourceCodeAccessor::FindExecutablePath()
 	// Search from JetBrainsToolbox folder
 	FString ToolboxBinPath;
 
-	if (FWindowsPlatformMisc::QueryRegKey(HKEY_CURRENT_USER, TEXT("Software\\JetBrains s.r.o.\\JetBrainsToolbox\\"), TEXT(""), ToolboxBinPath))
+	if (FWindowsPlatformMisc::QueryRegKey(HKEY_CURRENT_USER, TEXT("Software\\JetBrains\\Toolbox\\"), TEXT(""), ToolboxBinPath))
 	{
 		FPaths::NormalizeDirectoryName(ToolboxBinPath);
 		FString PatternString(TEXT("(.*)/bin"));
