CLionSourceCodeAccessor: UE4 -> UE5

#trivial

diff --git a/Engine/Plugins/Developer/CLionSourceCodeAccess/Source/CLionSourceCodeAccess/Private/CLionSourceCodeAccessor.cpp b/Engine/Plugins/Developer/CLionSourceCodeAccess/Source/CLionSourceCodeAccess/Private/CLionSourceCodeAccessor.cpp
--- a/Engine/Plugins/Developer/CLionSourceCodeAccess/Source/CLionSourceCodeAccess/Private/CLionSourceCodeAccessor.cpp
+++ b/Engine/Plugins/Developer/CLionSourceCodeAccess/Source/CLionSourceCodeAccess/Private/CLionSourceCodeAccessor.cpp
@@ -118,11 +118,11 @@ bool FCLionSourceCodeAccessor::OpenSolutionAtPath(const FString& InSolutionPath)
 	}
 
 	FString CorrectSolutionPath = InSolutionPath;
-	if (InSolutionPath.EndsWith(TEXT("UE4")))
+	if (InSolutionPath.EndsWith(TEXT("UE5")))
 	{
 		CorrectSolutionPath.LeftInline(CorrectSolutionPath.Len() - 3);
 	}
-	// UE4 passes the project folder and name, so strip the name off
+	// UE5 passes the project folder and name, so strip the name off
 	int32 LastPathIndex = CorrectSolutionPath.Find(TEXT("/"), ESearchCase::CaseSensitive, ESearchDir::FromEnd);
 	if (LastPathIndex > -1)
 	{
