Added support for Linux to search (using which) for CLion outside of the default expectated location

diff --git a/Engine/Plugins/Developer/CLionSourceCodeAccess/Source/CLionSourceCodeAccess/Private/CLionSourceCodeAccessor.cpp b/Engine/Plugins/Developer/CLionSourceCodeAccess/Source/CLionSourceCodeAccess/Private/CLionSourceCodeAccessor.cpp
--- a/Engine/Plugins/Developer/CLionSourceCodeAccess/Source/CLionSourceCodeAccess/Private/CLionSourceCodeAccessor.cpp
+++ b/Engine/Plugins/Developer/CLionSourceCodeAccess/Source/CLionSourceCodeAccess/Private/CLionSourceCodeAccessor.cpp
@@ -272,10 +272,34 @@ FString FCLionSourceCodeAccessor::FindExecutablePath()
 #else
 
 	// Linux Default Install
-	if(FPaths::FileExists(TEXT("/opt/clion/bin/clion.sh")))
+	if (FPaths::FileExists(TEXT("/opt/clion/bin/clion.sh")))
 	{
 		return TEXT("/opt/clion/bin/clion.sh");
 	}
+	
+	const TCHAR* Command[] = { TEXT("/usr/bin/which"), TEXT("clion.sh") };
+	
+	// Check to see if CLion is in the $PATH
+	FString CLionPath;
+	void* ReadPipe = nullptr;
+	void* WritePipe = nullptr;
+	FPlatformProcess::CreatePipe(ReadPipe, WritePipe);
+
+	FProcHandle ProcHandle = FPlatformProcess::CreateProc(Command[0], Command[1], false, false, false, nullptr, 0, nullptr, WritePipe);
+	if (ProcHandle.IsValid())
+	{
+		FPlatformProcess::WaitForProc(ProcHandle);
+		CLionPath = FPlatformProcess::ReadPipe(ReadPipe);
+		CLionPath = CLionPath.TrimChar('\n');
+	}
+	FPlatformProcess::ClosePipe(ReadPipe, WritePipe);
+	FPlatformProcess::CloseProc(ProcHandle);
+
+	if (!CLionPath.IsEmpty() && FPaths::FileExists(CLionPath))
+	{
+		return CLionPath;
+	}
+
 #endif
 
 	// Nothing was found, return nothing as well
