Check more places for CLion before giving up.

diff --git a/Engine/Plugins/Developer/CLionSourceCodeAccess/Source/CLionSourceCodeAccess/Private/CLionSourceCodeAccessor.cpp b/Engine/Plugins/Developer/CLionSourceCodeAccess/Source/CLionSourceCodeAccess/Private/CLionSourceCodeAccessor.cpp
--- a/Engine/Plugins/Developer/CLionSourceCodeAccess/Source/CLionSourceCodeAccess/Private/CLionSourceCodeAccessor.cpp
+++ b/Engine/Plugins/Developer/CLionSourceCodeAccess/Source/CLionSourceCodeAccess/Private/CLionSourceCodeAccessor.cpp
@@ -179,6 +179,30 @@ bool FCLionSourceCodeAccessor::SaveAllOpenDocuments() const
 	return false;
 }
 
+#if PLATFORM_LINUX
+
+// Returns true if a valid Exec path was found therefore OutExecPath is correctly set
+static bool GetExecPathFromDesktopFile(const FString& DesktopFilePath, FString& OutExecPath)
+{
+	if (!FPaths::FileExists(DesktopFilePath))
+	{
+		return false;
+	}
+	
+	FString DesktopFileContents;
+	FFileHelper::LoadFileToString(DesktopFileContents, *DesktopFilePath);
+
+	const bool Result = FParse::Value(*DesktopFileContents, TEXT("Exec="), OutExecPath);
+	if (!FPaths::FileExists(OutExecPath))
+	{
+		return false;
+	}
+	
+	return Result;
+}
+
+#endif
+
 FString FCLionSourceCodeAccessor::FindExecutablePath()
 {
 #if PLATFORM_WINDOWS
@@ -270,34 +294,43 @@ FString FCLionSourceCodeAccessor::FindExecutablePath()
 	}
 
 #else
-
-	// Linux Default Install
+	
+	// Potential executable location
 	if (FPaths::FileExists(TEXT("/opt/clion/bin/clion.sh")))
 	{
 		return TEXT("/opt/clion/bin/clion.sh");
 	}
 	
-	const TCHAR* Command[] = { TEXT("/usr/bin/which"), TEXT("clion.sh") };
+	FString ClionPath;
 	
-	// Check to see if CLion is in the $PATH
-	FString CLionPath;
-	void* ReadPipe = nullptr;
-	void* WritePipe = nullptr;
-	FPlatformProcess::CreatePipe(ReadPipe, WritePipe);
-
-	FProcHandle ProcHandle = FPlatformProcess::CreateProc(Command[0], Command[1], false, false, false, nullptr, 0, nullptr, WritePipe);
-	if (ProcHandle.IsValid())
+	// First potential .desktop file location
+	FString DesktopFilePath = TEXT("/usr/share/applications/jetbrains-clion.desktop");
+	if (GetExecPathFromDesktopFile(DesktopFilePath, ClionPath))
 	{
-		FPlatformProcess::WaitForProc(ProcHandle);
-		CLionPath = FPlatformProcess::ReadPipe(ReadPipe);
-		CLionPath = CLionPath.TrimChar('\n');
+		return ClionPath;
 	}
-	FPlatformProcess::ClosePipe(ReadPipe, WritePipe);
-	FPlatformProcess::CloseProc(ProcHandle);
-
-	if (!CLionPath.IsEmpty() && FPaths::FileExists(CLionPath))
+	
+	// Second potential .desktop file location
+	DesktopFilePath = FPaths::Combine(FPlatformMisc::GetEnvironmentVariable(TEXT("HOME")), TEXT("/.local/share/applications/jetbrains-clion.desktop"));
+	if (GetExecPathFromDesktopFile(DesktopFilePath, ClionPath))
 	{
-		return CLionPath;
+		return ClionPath;
+	}
+	
+	// Check to see if CLion is in the $PATH
+	TArray<FString> PathArray;
+	FPlatformMisc::GetEnvironmentVariable(TEXT("PATH")).ParseIntoArray(PathArray, FPlatformMisc::GetPathVarDelimiter());
+	
+	for (const FString& Path : PathArray)
+	{
+		if (Path.Contains("CLion"))
+		{
+			ClionPath = FPaths::Combine(Path, TEXT("clion.sh"));
+			if (FPaths::FileExists(ClionPath))
+			{
+				return ClionPath;
+			}
+		}
 	}
 
 #endif
